- Fix "-allmodules" build failures caused by "VoiceChat" plugin being built, by changing the VoiceChat module to "External" type. It is a header-only module, so External should stop the build system trying to build it.
- Fix Xbox/IOS build failures caused by incorrect header paths, and incorrect use of FVoiceChatResult

diff --git a/Engine/Plugins/Online/VoiceChat/VoiceChat/Source/Public/VoiceChatErrors.h b/Engine/Plugins/Online/VoiceChat/VoiceChat/Source/Public/VoiceChatErrors.h
--- a/Engine/Plugins/Online/VoiceChat/VoiceChat/Source/Public/VoiceChatErrors.h
+++ b/Engine/Plugins/Online/VoiceChat/VoiceChat/Source/Public/VoiceChatErrors.h
@@ -17,6 +17,7 @@ namespace VoiceChat
 		inline FVoiceChatResult NotInChannel() { return VOICECHAT_ERROR(EVoiceChatResult::InvalidState, TEXT("not_in_channel")); }
 		inline FVoiceChatResult ChannelJoinInProgress() { return VOICECHAT_ERROR(EVoiceChatResult::InvalidState, TEXT("channel_join_in_progress")); }
 		inline FVoiceChatResult ChannelLeaveInProgress() { return VOICECHAT_ERROR(EVoiceChatResult::InvalidState, TEXT("channel_leave_in_progress")); }
+		inline FVoiceChatResult PlatformPartyChatActive() { return VOICECHAT_ERROR(EVoiceChatResult::InvalidState, TEXT("platform_party_chat_active")); }
 
 		inline FVoiceChatResult InvalidArgument(const FString& ErrorDesc = FString()) { return VOICECHAT_ERROR(EVoiceChatResult::InvalidArgument, TEXT("invalid_argument"), ErrorDesc); }
 		inline FVoiceChatResult CredentialsInvalid() { return VOICECHAT_ERROR(EVoiceChatResult::CredentialsInvalid, TEXT("credentials_invalid")); }
diff --git a/Engine/Plugins/Online/VoiceChat/VoiceChat/Source/VoiceChat.Build.cs b/Engine/Plugins/Online/VoiceChat/VoiceChat/Source/VoiceChat.Build.cs
--- a/Engine/Plugins/Online/VoiceChat/VoiceChat/Source/VoiceChat.Build.cs
+++ b/Engine/Plugins/Online/VoiceChat/VoiceChat/Source/VoiceChat.Build.cs
@@ -8,7 +8,8 @@ namespace UnrealBuildTool.Rules
 	{
 		public VoiceChat(ReadOnlyTargetRules Target) : base(Target)
 		{
-			PCHUsage = ModuleRules.PCHUsageMode.UseExplicitOrSharedPCHs;
+			// External so we don't build a DLL for this, it is header only.
+			Type = ModuleType.External;
 
 			PublicDependencyModuleNames.AddRange(
 				new string[]
