- Fix build failures caused by VoiceChat external module not explicitly stating its public include path. This was switched from a CPlusPlus module type where the Public/Private paths are implied.

diff --git a/Engine/Plugins/Online/VoiceChat/VoiceChat/Source/VoiceChat.Build.cs b/Engine/Plugins/Online/VoiceChat/VoiceChat/Source/VoiceChat.Build.cs
--- a/Engine/Plugins/Online/VoiceChat/VoiceChat/Source/VoiceChat.Build.cs
+++ b/Engine/Plugins/Online/VoiceChat/VoiceChat/Source/VoiceChat.Build.cs
@@ -11,6 +11,8 @@ namespace UnrealBuildTool.Rules
 			// External so we don't build a DLL for this, it is header only.
 			Type = ModuleType.External;
 
+			PublicIncludePaths.Add(Path.Combine(ModuleDirectory, "Public"));
+
 			PublicDependencyModuleNames.AddRange(
 				new string[]
 				{
