Improve API comment

diff --git a/Engine/Plugins/Online/VoiceChat/VoiceChat/Source/Public/VoiceChatResult.h b/Engine/Plugins/Online/VoiceChat/VoiceChat/Source/Public/VoiceChatResult.h
--- a/Engine/Plugins/Online/VoiceChat/VoiceChat/Source/Public/VoiceChatResult.h
+++ b/Engine/Plugins/Online/VoiceChat/VoiceChat/Source/Public/VoiceChatResult.h
@@ -29,7 +29,7 @@ enum class EVoiceChatResult
 
 struct FVoiceChatResult
 {
-	/** Whether or not we succeeded */
+	/** Success, or an error category */
 	EVoiceChatResult ResultCode;
 	/** If we failed, the code for the error */
 	FString ErrorCode;
