VoiceChat modifications

- Promote MissingConfig to VoiceChatErrors from VivoxVoiceChatErrors
- Add AlreadyInProgress, NotPermitted, Throttled to VoiceChatErrors
- Add NotPermitted, Throttled EVoiceChatResult categories

diff --git a/Engine/Plugins/Online/VoiceChat/VivoxVoiceChat/Source/Private/VivoxVoiceChat.cpp b/Engine/Plugins/Online/VoiceChat/VivoxVoiceChat/Source/Private/VivoxVoiceChat.cpp
--- a/Engine/Plugins/Online/VoiceChat/VivoxVoiceChat/Source/Private/VivoxVoiceChat.cpp
+++ b/Engine/Plugins/Online/VoiceChat/VivoxVoiceChat/Source/Private/VivoxVoiceChat.cpp
@@ -554,7 +554,7 @@ void FVivoxVoiceChat::Connect(const FOnVoiceChatConnectCompleteDelegate& Delegat
 			if (VivoxServerUrl.IsEmpty() || VivoxDomain.IsEmpty() || VivoxNamespace.IsEmpty())
 			{
 				UE_LOG(LogVivoxVoiceChat, Warning, TEXT("[VoiceChat.Vivox] ServerUrl, Domain, or Issuer is not set. Vivox voice chat will not work"));
-				Result = VivoxVoiceChat::Errors::MissingConfig();
+				Result = VoiceChat::Errors::MissingConfig();
 				Result.ErrorDesc = FString::Printf(TEXT("ServerUrl=[%s] Domain=[%s] Namespace=[%s]"), *VivoxServerUrl, *VivoxDomain, *VivoxNamespace);
 			}
 			else
diff --git a/Engine/Plugins/Online/VoiceChat/VivoxVoiceChat/Source/Public/VivoxVoiceChatErrors.h b/Engine/Plugins/Online/VoiceChat/VivoxVoiceChat/Source/Public/VivoxVoiceChatErrors.h
--- a/Engine/Plugins/Online/VoiceChat/VivoxVoiceChat/Source/Public/VivoxVoiceChatErrors.h
+++ b/Engine/Plugins/Online/VoiceChat/VivoxVoiceChat/Source/Public/VivoxVoiceChatErrors.h
@@ -7,7 +7,6 @@ namespace VivoxVoiceChat
 {
 	namespace Errors
 	{
-		inline FVoiceChatResult MissingConfig() { return VIVOXVOICECHAT_ERROR(EVoiceChatResult::InvalidState, TEXT("missing_config")); }
 		inline FVoiceChatResult AlreadyInitialized() { return VIVOXVOICECHAT_ERROR(EVoiceChatResult::InvalidState, TEXT("already_initialized")); }
 		inline FVoiceChatResult AlreadyLoggedIn() { return VIVOXVOICECHAT_ERROR(EVoiceChatResult::InvalidState, TEXT("already_logged_in")); }
 		inline FVoiceChatResult KickedFromChannel() { return VIVOXVOICECHAT_ERROR(EVoiceChatResult::ImplementationError, TEXT("kicked_from_channel")); }
diff --git a/Engine/Plugins/Online/VoiceChat/VoiceChat/Source/Public/VoiceChatErrors.h b/Engine/Plugins/Online/VoiceChat/VoiceChat/Source/Public/VoiceChatErrors.h
--- a/Engine/Plugins/Online/VoiceChat/VoiceChat/Source/Public/VoiceChatErrors.h
+++ b/Engine/Plugins/Online/VoiceChat/VoiceChat/Source/Public/VoiceChatErrors.h
@@ -11,7 +11,12 @@ namespace VoiceChat
 		inline FVoiceChatResult NotConnected() { return VOICECHAT_ERROR(EVoiceChatResult::NotConnected, TEXT("not_connected")); }
 		inline FVoiceChatResult NotLoggedIn() { return VOICECHAT_ERROR(EVoiceChatResult::NotLoggedIn, TEXT("not_logged_in")); }
 
+		inline FVoiceChatResult NotPermitted() { return VOICECHAT_ERROR(EVoiceChatResult::NotPermitted, TEXT("not_permitted")); }
+		inline FVoiceChatResult Throttled()	{ return VOICECHAT_ERROR(EVoiceChatResult::Throttled, TEXT("throttled")); }
+
 		inline FVoiceChatResult InvalidState() { return VOICECHAT_ERROR(EVoiceChatResult::InvalidState, TEXT("invalid_state")); }
+		inline FVoiceChatResult MissingConfig() { return VOICECHAT_ERROR(EVoiceChatResult::InvalidState, TEXT("missing_config")); }
+		inline FVoiceChatResult AlreadyInProgress() { return VOICECHAT_ERROR(EVoiceChatResult::InvalidState, TEXT("already_in_progress")); }
 		inline FVoiceChatResult DisconnectInProgress() { return VOICECHAT_ERROR(EVoiceChatResult::InvalidState, TEXT("disconnect_in_progress")); }
 		inline FVoiceChatResult OtherUserLoggedIn() { return VOICECHAT_ERROR(EVoiceChatResult::InvalidState, TEXT("other_user_logged_in")); }
 		inline FVoiceChatResult NotInChannel() { return VOICECHAT_ERROR(EVoiceChatResult::InvalidState, TEXT("not_in_channel")); }
@@ -20,7 +25,7 @@ namespace VoiceChat
 		inline FVoiceChatResult PlatformPartyChatActive() { return VOICECHAT_ERROR(EVoiceChatResult::InvalidState, TEXT("platform_party_chat_active")); }
 
 		inline FVoiceChatResult InvalidArgument(const FString& ErrorDesc = FString()) { return VOICECHAT_ERROR(EVoiceChatResult::InvalidArgument, TEXT("invalid_argument"), ErrorDesc); }
-		inline FVoiceChatResult CredentialsInvalid() { return VOICECHAT_ERROR(EVoiceChatResult::CredentialsInvalid, TEXT("credentials_invalid")); }
+		inline FVoiceChatResult CredentialsInvalid(const FString& ErrorDesc = FString()) { return VOICECHAT_ERROR(EVoiceChatResult::CredentialsInvalid, TEXT("credentials_invalid"), ErrorDesc); }
 		inline FVoiceChatResult CredentialsExpired() { return VOICECHAT_ERROR(EVoiceChatResult::CredentialsExpired, TEXT("credentials_expired")); }
 
 		inline FVoiceChatResult ClientTimeout() { return VOICECHAT_ERROR(EVoiceChatResult::ClientTimeout, TEXT("client_timeout")); }
diff --git a/Engine/Plugins/Online/VoiceChat/VoiceChat/Source/Public/VoiceChatResult.h b/Engine/Plugins/Online/VoiceChat/VoiceChat/Source/Public/VoiceChatResult.h
--- a/Engine/Plugins/Online/VoiceChat/VoiceChat/Source/Public/VoiceChatResult.h
+++ b/Engine/Plugins/Online/VoiceChat/VoiceChat/Source/Public/VoiceChatResult.h
@@ -11,6 +11,8 @@ enum class EVoiceChatResult
 	NotInitialized,
 	NotConnected,
 	NotLoggedIn,
+	NotPermitted,
+	Throttled,
 
 	// Common argument errors
 	InvalidArgument,
