Fixing bad merge

diff --git a/Engine/Plugins/Online/VoiceChat/VoiceChat/Source/Public/VoiceChat.h b/Engine/Plugins/Online/VoiceChat/VoiceChat/Source/Public/VoiceChat.h
--- a/Engine/Plugins/Online/VoiceChat/VoiceChat/Source/Public/VoiceChat.h
+++ b/Engine/Plugins/Online/VoiceChat/VoiceChat/Source/Public/VoiceChat.h
@@ -113,7 +113,7 @@ DECLARE_MULTICAST_DELEGATE_ThreeParams(FOnVoiceChatPlayerVolumeUpdatedDelegate,
 DECLARE_MULTICAST_DELEGATE_ThreeParams(FOnVoiceChatRecordSamplesAvailableDelegate, TArrayView<const int16> /* PcmSamples */, int /* SampleRate */, int /* Channels */);
 DECLARE_MULTICAST_DELEGATE_ThreeParams(FOnVoiceChatAfterCaptureAudioReadDelegate, TArrayView<int16> /* PcmSamples */, int /* SampleRate */, int /* Channels */);
 DECLARE_MULTICAST_DELEGATE_FourParams(FOnVoiceChatBeforeCaptureAudioSentDelegate, TArrayView<const int16> /* PcmSamples */, int /* SampleRate */, int /* Channels */, bool /* bIsSpeaking */);
-DECLARE_MULTICAST_DELEGATE_FourParams(FOnVoiceChatBeforeRecvAudioRenderedDelegate, TArrayView<int16> /* PcmSamples */, int /* SampleRate */, int /* Channels */, bool /* bIsSilence */);
+DECLARE_MULTICAST_DELEGATE_FiveParams(FOnVoiceChatBeforeRecvAudioRenderedDelegate, TArrayView<int16> /* PcmSamples */, int /* SampleRate */, int /* Channels */, bool /* bIsSilence */, const FString& /* PlayerName */);
 
 class IVoiceChatUser
 {
