[VoiceChat]
* VoiceChatErrors - Added missing include

diff --git a/Engine/Plugins/Online/VoiceChat/VoiceChat/Source/Public/VoiceChatErrors.h b/Engine/Plugins/Online/VoiceChat/VoiceChat/Source/Public/VoiceChatErrors.h
--- a/Engine/Plugins/Online/VoiceChat/VoiceChat/Source/Public/VoiceChatErrors.h
+++ b/Engine/Plugins/Online/VoiceChat/VoiceChat/Source/Public/VoiceChatErrors.h
@@ -1,6 +1,8 @@
 // Copyright Epic Games, Inc. All Rights Reserved.
 #pragma once
 
+#include "VoiceChatResult.h"
+
 #define VOICECHAT_ERROR(...) FVoiceChatResult::CreateError(TEXT("errors.com.epicgames.voicechat"), __VA_ARGS__)
 
 namespace VoiceChat
