Added multichannel transmit for IVoiceChat

#8974

diff --git a/Engine/Plugins/Online/OnlineSubsystemEOS/Source/OnlineSubsystemEOS/Private/OnlineSubsystemEOS.cpp b/Engine/Plugins/Online/OnlineSubsystemEOS/Source/OnlineSubsystemEOS/Private/OnlineSubsystemEOS.cpp
--- a/Engine/Plugins/Online/OnlineSubsystemEOS/Source/OnlineSubsystemEOS/Private/OnlineSubsystemEOS.cpp
+++ b/Engine/Plugins/Online/OnlineSubsystemEOS/Source/OnlineSubsystemEOS/Private/OnlineSubsystemEOS.cpp
@@ -104,9 +104,9 @@ public:
 	virtual FOnVoiceChatPlayerVolumeUpdatedDelegate& OnVoiceChatPlayerVolumeUpdated() override { return VoiceChatUser.OnVoiceChatPlayerVolumeUpdated(); }
 	virtual void TransmitToAllChannels() override { VoiceChatUser.TransmitToAllChannels(); }
 	virtual void TransmitToNoChannels() override { VoiceChatUser.TransmitToNoChannels(); }
-	virtual void TransmitToSpecificChannel(const FString& ChannelName) override { VoiceChatUser.TransmitToSpecificChannel(ChannelName); }
+	virtual void TransmitToSpecificChannels(const TSet<FString>& ChannelNames) override { VoiceChatUser.TransmitToSpecificChannels(ChannelNames); }
 	virtual EVoiceChatTransmitMode GetTransmitMode() const override { return VoiceChatUser.GetTransmitMode(); }
-	virtual FString GetTransmitChannel() const override { return VoiceChatUser.GetTransmitChannel(); }
+	virtual TSet<FString> GetTransmitChannels() const override { return VoiceChatUser.GetTransmitChannels(); }
 	virtual FDelegateHandle StartRecording(const FOnVoiceChatRecordSamplesAvailableDelegate::FDelegate& Delegate) override { return VoiceChatUser.StartRecording(Delegate); }
 	virtual void StopRecording(FDelegateHandle Handle) override { VoiceChatUser.StopRecording(Handle); }
 	virtual FDelegateHandle RegisterOnVoiceChatAfterCaptureAudioReadDelegate(const FOnVoiceChatAfterCaptureAudioReadDelegate::FDelegate& Delegate) override { return VoiceChatUser.RegisterOnVoiceChatAfterCaptureAudioReadDelegate(Delegate); }
diff --git a/Engine/Plugins/Online/VoiceChat/EOSVoiceChat/Source/EOSVoiceChat/Private/EOSVoiceChat.cpp b/Engine/Plugins/Online/VoiceChat/EOSVoiceChat/Source/EOSVoiceChat/Private/EOSVoiceChat.cpp
--- a/Engine/Plugins/Online/VoiceChat/EOSVoiceChat/Source/EOSVoiceChat/Private/EOSVoiceChat.cpp
+++ b/Engine/Plugins/Online/VoiceChat/EOSVoiceChat/Source/EOSVoiceChat/Private/EOSVoiceChat.cpp
@@ -703,9 +703,9 @@ void FEOSVoiceChat::TransmitToNoChannels()
 	GetVoiceChatUser().TransmitToNoChannels();
 }
 
-void FEOSVoiceChat::TransmitToSpecificChannel(const FString& ChannelName)
+void FEOSVoiceChat::TransmitToSpecificChannels(const TSet<FString>& ChannelNames)
 {
-	GetVoiceChatUser().TransmitToSpecificChannel(ChannelName);
+	GetVoiceChatUser().TransmitToSpecificChannels(ChannelNames);
 }
 
 EVoiceChatTransmitMode FEOSVoiceChat::GetTransmitMode() const
@@ -713,9 +713,9 @@ EVoiceChatTransmitMode FEOSVoiceChat::GetTransmitMode() const
 	return GetVoiceChatUser().GetTransmitMode();
 }
 
-FString FEOSVoiceChat::GetTransmitChannel() const
+TSet<FString> FEOSVoiceChat::GetTransmitChannels() const
 {
-	return GetVoiceChatUser().GetTransmitChannel();
+	return GetVoiceChatUser().GetTransmitChannels();
 }
 
 FDelegateHandle FEOSVoiceChat::StartRecording(const FOnVoiceChatRecordSamplesAvailableDelegate::FDelegate& Delegate)
diff --git a/Engine/Plugins/Online/VoiceChat/EOSVoiceChat/Source/EOSVoiceChat/Private/EOSVoiceChatUser.cpp b/Engine/Plugins/Online/VoiceChat/EOSVoiceChat/Source/EOSVoiceChat/Private/EOSVoiceChatUser.cpp
--- a/Engine/Plugins/Online/VoiceChat/EOSVoiceChat/Source/EOSVoiceChat/Private/EOSVoiceChatUser.cpp
+++ b/Engine/Plugins/Online/VoiceChat/EOSVoiceChat/Source/EOSVoiceChat/Private/EOSVoiceChatUser.cpp
@@ -800,13 +800,14 @@ void FEOSVoiceChatUser::TransmitToNoChannels()
 	}
 }
 
-void FEOSVoiceChatUser::TransmitToSpecificChannel(const FString& Channel)
+void FEOSVoiceChatUser::TransmitToSpecificChannels(const TSet<FString>& ChannelNames)
 {
-	if (TransmitState.Mode != EVoiceChatTransmitMode::Channel ||
-		TransmitState.ChannelName != Channel)
+	if (TransmitState.Mode != EVoiceChatTransmitMode::SpecificChannels || 
+		!TransmitState.SpecificChannels.Difference(ChannelNames).IsEmpty() || 
+		!ChannelNames.Difference(TransmitState.SpecificChannels).IsEmpty())
 	{
-		TransmitState.Mode = EVoiceChatTransmitMode::Channel;
-		TransmitState.ChannelName = Channel;
+		TransmitState.Mode = EVoiceChatTransmitMode::SpecificChannels;
+		TransmitState.SpecificChannels = ChannelNames;
 		ApplySendingOptions();
 	}
 }
@@ -816,14 +817,13 @@ EVoiceChatTransmitMode FEOSVoiceChatUser::GetTransmitMode() const
 	return TransmitState.Mode;
 }
 
-FString FEOSVoiceChatUser::GetTransmitChannel() const
+TSet<FString> FEOSVoiceChatUser::GetTransmitChannels() const
 {
-	FString TransmitChannel;
-	if (TransmitState.Mode == EVoiceChatTransmitMode::Channel)
+	if (TransmitState.Mode == EVoiceChatTransmitMode::SpecificChannels)
 	{
-		TransmitChannel = TransmitState.ChannelName;
+		return TransmitState.SpecificChannels;
 	}
-	return TransmitChannel;
+	return TSet<FString>();
 }
 
 FDelegateHandle FEOSVoiceChatUser::StartRecording(const FOnVoiceChatRecordSamplesAvailableDelegate::FDelegate& Delegate)
@@ -1308,8 +1308,8 @@ void FEOSVoiceChatUser::ApplySendingOptions()
 
 void FEOSVoiceChatUser::ApplySendingOptions(FChannelSession& ChannelSession)
 {
-	const bool bCanTransmitToChannel = TransmitState.Mode == EVoiceChatTransmitMode::All ||
-		(TransmitState.Mode == EVoiceChatTransmitMode::Channel && TransmitState.ChannelName == ChannelSession.ChannelName);
+	const bool bCanTransmitToChannel = TransmitState.Mode == EVoiceChatTransmitMode::All
+	|| (TransmitState.Mode == EVoiceChatTransmitMode::SpecificChannels && TransmitState.SpecificChannels.Contains(ChannelSession.ChannelName));
 
 	ChannelSession.DesiredSendingState.bAudioEnabled = bCanTransmitToChannel && !AudioInputOptions.bMuted && !ChannelSession.bIsNotListening;
 
@@ -2415,8 +2415,8 @@ bool FEOSVoiceChatUser::Exec(UWorld* InWorld, const TCHAR* Cmd, FOutputDevice& A
 					case EVoiceChatTransmitMode::None:
 						TransmitString = TEXT("NONE");
 						break;
-					case EVoiceChatTransmitMode::Channel:
-						TransmitString = FString::Printf(TEXT("CHANNEL:%s"), *GetTransmitChannel());
+					case EVoiceChatTransmitMode::SpecificChannels:
+						TransmitString = FString::Printf(TEXT("CHANNELS:[%s]"), *FString::Join(GetTransmitChannels(), TEXT(", ")));
 						break;
 					}
 					EOS_EXEC_LOG(TEXT("    Channels: transmitting:%s"), *TransmitString);
@@ -2610,7 +2610,7 @@ bool FEOSVoiceChatUser::Exec(UWorld* InWorld, const TCHAR* Cmd, FOutputDevice& A
 			FString ChannelName;
 			if (FParse::Value(Cmd, TEXT("ChannelName="), ChannelName))
 			{
-				TransmitToSpecificChannel(ChannelName);
+				TransmitToSpecificChannels({ ChannelName });
 				return true;
 			}
 		}
diff --git a/Engine/Plugins/Online/VoiceChat/EOSVoiceChat/Source/EOSVoiceChat/Public/EOSVoiceChat.h b/Engine/Plugins/Online/VoiceChat/EOSVoiceChat/Source/EOSVoiceChat/Public/EOSVoiceChat.h
--- a/Engine/Plugins/Online/VoiceChat/EOSVoiceChat/Source/EOSVoiceChat/Public/EOSVoiceChat.h
+++ b/Engine/Plugins/Online/VoiceChat/EOSVoiceChat/Source/EOSVoiceChat/Public/EOSVoiceChat.h
@@ -123,9 +123,9 @@ public:
 	virtual FOnVoiceChatPlayerVolumeUpdatedDelegate& OnVoiceChatPlayerVolumeUpdated() override;
 	virtual void TransmitToAllChannels() override;
 	virtual void TransmitToNoChannels() override;
-	virtual void TransmitToSpecificChannel(const FString& ChannelName) override;
+	virtual void TransmitToSpecificChannels(const TSet<FString>& ChannelNames) override;
 	virtual EVoiceChatTransmitMode GetTransmitMode() const override;
-	virtual FString GetTransmitChannel() const override;
+	virtual TSet<FString> GetTransmitChannels() const override;
 	virtual FDelegateHandle StartRecording(const FOnVoiceChatRecordSamplesAvailableDelegate::FDelegate& Delegate) override;
 	virtual void StopRecording(FDelegateHandle Handle) override;
 	virtual FDelegateHandle RegisterOnVoiceChatAfterCaptureAudioReadDelegate(const FOnVoiceChatAfterCaptureAudioReadDelegate::FDelegate& Delegate) override;
diff --git a/Engine/Plugins/Online/VoiceChat/EOSVoiceChat/Source/EOSVoiceChat/Public/EOSVoiceChatUser.h b/Engine/Plugins/Online/VoiceChat/EOSVoiceChat/Source/EOSVoiceChat/Public/EOSVoiceChatUser.h
--- a/Engine/Plugins/Online/VoiceChat/EOSVoiceChat/Source/EOSVoiceChat/Public/EOSVoiceChatUser.h
+++ b/Engine/Plugins/Online/VoiceChat/EOSVoiceChat/Source/EOSVoiceChat/Public/EOSVoiceChatUser.h
@@ -71,9 +71,9 @@ public:
 	virtual FOnVoiceChatPlayerVolumeUpdatedDelegate& OnVoiceChatPlayerVolumeUpdated() override { return OnVoiceChatPlayerVolumeUpdatedDelegate; }
 	virtual void TransmitToAllChannels() override;
 	virtual void TransmitToNoChannels() override;
-	virtual void TransmitToSpecificChannel(const FString& ChannelName) override;
+	virtual void TransmitToSpecificChannels(const TSet<FString>& ChannelNames) override;
 	virtual EVoiceChatTransmitMode GetTransmitMode() const override;
-	virtual FString GetTransmitChannel() const override;
+	virtual TSet<FString> GetTransmitChannels() const override;
 	virtual FDelegateHandle StartRecording(const FOnVoiceChatRecordSamplesAvailableDelegate::FDelegate& Delegate) override;
 	virtual void StopRecording(FDelegateHandle Handle) override;
 	virtual FDelegateHandle RegisterOnVoiceChatAfterCaptureAudioReadDelegate(const FOnVoiceChatAfterCaptureAudioReadDelegate::FDelegate& Delegate) override;
@@ -266,7 +266,7 @@ protected:
 	struct FTransmitState
 	{
 		EVoiceChatTransmitMode Mode = EVoiceChatTransmitMode::All;
-		FString ChannelName;
+		TSet<FString> SpecificChannels;
 	};
 	FTransmitState TransmitState;
 
diff --git a/Engine/Plugins/Online/VoiceChat/VoiceChat/Source/Public/VoiceChat.h b/Engine/Plugins/Online/VoiceChat/VoiceChat/Source/Public/VoiceChat.h
--- a/Engine/Plugins/Online/VoiceChat/VoiceChat/Source/Public/VoiceChat.h
+++ b/Engine/Plugins/Online/VoiceChat/VoiceChat/Source/Public/VoiceChat.h
@@ -49,7 +49,9 @@ enum class EVoiceChatTransmitMode
 	/** Transmit to all of the channels you are in */
 	All,
 	/** Transmit to a specific channel */
-	Channel
+	Channel UE_DEPRECATED(5.3, "EVoiceChatTransmitMode::Channel is deprecated. Please use EVoiceChatTransmitMode::SpecificChannels instead!"),
+	/** Transmit to specific channel(s) */
+	SpecificChannels,
 };
 
 struct FVoiceChatCallStats
@@ -570,7 +572,7 @@ public:
 	 * Transmit to all channels that you are currently in
 	 */
 	virtual void TransmitToAllChannels() = 0;
-
+	
 	/**
 	 * Stop transmitting to all channels that you are currently in. This does not apply to channels you will join.
 	 */
@@ -581,21 +583,48 @@ public:
 	 *
 	 * @param ChannelName Channel to transmit to
 	 */
-	virtual void TransmitToSpecificChannel(const FString& ChannelName) = 0;
+	UE_DEPRECATED(5.3, "TransmitToSpecificChannel is deprecated. Please use TransmitToSpecificChannels instead!")
+	virtual void TransmitToSpecificChannel(const FString& ChannelName) { TransmitToSpecificChannels({ChannelName}); }
+
+	/**
+	 * Transmit to Specified set of channels
+	 * @param ChannelNames Channels to transmit to
+	 */
+	virtual void TransmitToSpecificChannels(const TSet<FString>& ChannelNames) = 0;
+
+	/**
+	 * Add specific channels to transmit
+	 * @param ChannelNames New channels to add to the current transmit set
+	 */
+	virtual void AddSpecificChannelsToTransmit(const TSet<FString>& ChannelNames) { TransmitToSpecificChannels(GetTransmitChannels().Union(ChannelNames)); }
+
+	/**
+	 * Remove specific channels from transmit
+	 * @param ChannelNames Channels to remove from the current transmit set
+	 */
+	virtual void RemoveSpecificChannelsFromTransmit(const TSet<FString>& ChannelNames) { TransmitToSpecificChannels(GetTransmitChannels().Difference(ChannelNames)); }
 
 	/**
 	 * Get the current transmit mode
 	 *
-	 * @return Transmit mode. If it is EVoiceChatTransmitMode::Channel, the channel can be retrieved using GetTransmitChannel()
+	 * @return Transmit mode. If it is EVoiceChatTransmitMode::SpecificChannels, the current channels can be retrieved using GetTransmitChannels()
 	 */
 	virtual EVoiceChatTransmitMode GetTransmitMode() const = 0;
 
 	/**
-	 * Get the specific channel we are transmitting to
+	 * Get the first channel we are transmitting to
+	 *
+	 * @return Channel we are transmitting to. Will be empty if GetTransmitMode() returned a value other than EVoiceChatTransmitMode::SpecificChannels
+	 */
+	UE_DEPRECATED(5.3, "GetTransmitChannel is deprecated. Please use GetTransmitChannels instead!")
+		virtual FString GetTransmitChannel() const { return ((GetTransmitMode() == EVoiceChatTransmitMode::SpecificChannels) && !GetTransmitChannels().IsEmpty()) ? GetTransmitChannels().Array()[0] : FString(); };
+
+	/**
+	 * Get all channels we are transmitting to
 	 *
-	 * @return Channel we are transmitting to. Will be empty if GetTransmissionMode() returned a value other than EVoiceChatTransmitMode::Channel
+	 * @return Channels	 we are transmitting to. Will be empty if GetTransmissionMode() returned a value other than EVoiceChatTransmitMode::SpecificChannels
 	 */
-	virtual FString GetTransmitChannel() const = 0;
+	virtual TSet<FString> GetTransmitChannels() const = 0;
 
 	/*
 	 * Start recording
