Fix include path warnings in external modules

#rnx

diff --git a/Engine/Plugins/Online/VoiceChat/VoiceChat/Source/VoiceChat.Build.cs b/Engine/Plugins/Online/VoiceChat/VoiceChat/Source/VoiceChat.Build.cs
--- a/Engine/Plugins/Online/VoiceChat/VoiceChat/Source/VoiceChat.Build.cs
+++ b/Engine/Plugins/Online/VoiceChat/VoiceChat/Source/VoiceChat.Build.cs
@@ -10,6 +10,7 @@ namespace UnrealBuildTool.Rules
 		{
 			// External so we don't build a DLL for this, it is header only.
 			Type = ModuleType.External;
+			ModuleIncludePathWarningLevel = WarningLevel.Off;
 
 			PublicIncludePaths.Add(Path.Combine(ModuleDirectory, "Public"));
 
