Adding channel name to BeforeCaptureAudioSent

diff --git a/Engine/Plugins/Online/VoiceChat/EOSVoiceChat/Source/EOSVoiceChat/Private/EOSVoiceChatUser.cpp b/Engine/Plugins/Online/VoiceChat/EOSVoiceChat/Source/EOSVoiceChat/Private/EOSVoiceChatUser.cpp
--- a/Engine/Plugins/Online/VoiceChat/EOSVoiceChat/Source/EOSVoiceChat/Private/EOSVoiceChatUser.cpp
+++ b/Engine/Plugins/Online/VoiceChat/EOSVoiceChat/Source/EOSVoiceChat/Private/EOSVoiceChatUser.cpp
@@ -2239,7 +2239,10 @@ void FEOSVoiceChatUser::OnChannelAudioBeforeSend(const EOS_RTCAudio_AudioBeforeS
 			// Allow any processes to modify audio through DSP effects processing
 			OnVoiceChatAfterCaptureAudioReadDelegate.Broadcast(WriteableSamples, Buffer->SampleRate, Buffer->Channels);
 
-			OnVoiceChatBeforeCaptureAudioSentDelegate.Broadcast(Samples, Buffer->SampleRate, Buffer->Channels, bSpeaking);
+			const FString ChannelName = UTF8_TO_TCHAR(CallbackInfo->RoomName);
+
+			OnVoiceChatBeforeCaptureAudioSentDelegate.Broadcast(ChannelName, Samples, Buffer->SampleRate, Buffer->Channels, bSpeaking);
+
 		}
 		else
 		{
diff --git a/Engine/Plugins/Online/VoiceChat/VoiceChat/Source/Public/VoiceChat.h b/Engine/Plugins/Online/VoiceChat/VoiceChat/Source/Public/VoiceChat.h
--- a/Engine/Plugins/Online/VoiceChat/VoiceChat/Source/Public/VoiceChat.h
+++ b/Engine/Plugins/Online/VoiceChat/VoiceChat/Source/Public/VoiceChat.h
@@ -112,7 +112,7 @@ DECLARE_MULTICAST_DELEGATE_ThreeParams(FOnVoiceChatPlayerVolumeUpdatedDelegate,
 
 DECLARE_MULTICAST_DELEGATE_ThreeParams(FOnVoiceChatRecordSamplesAvailableDelegate, TArrayView<const int16> /* PcmSamples */, int /* SampleRate */, int /* Channels */);
 DECLARE_MULTICAST_DELEGATE_ThreeParams(FOnVoiceChatAfterCaptureAudioReadDelegate, TArrayView<int16> /* PcmSamples */, int /* SampleRate */, int /* Channels */);
-DECLARE_MULTICAST_DELEGATE_FourParams(FOnVoiceChatBeforeCaptureAudioSentDelegate, TArrayView<const int16> /* PcmSamples */, int /* SampleRate */, int /* Channels */, bool /* bIsSpeaking */);
+DECLARE_MULTICAST_DELEGATE_FiveParams(FOnVoiceChatBeforeCaptureAudioSentDelegate, const FString& /* ChannelName */, TArrayView<const int16> /* PcmSamples */, int /* SampleRate */, int /* Channels */, bool /* bIsSpeaking */);
 DECLARE_MULTICAST_DELEGATE_SixParams(FOnVoiceChatBeforeRecvAudioRenderedDelegate, TArrayView<int16> /* PcmSamples */, int /* SampleRate */, int /* Channels */, bool /* bIsSilence */, const FString& /*ChannelName*/, const FString& /* PlayerName */);
 
 class IVoiceChatUser
