- Add ChannelName to AfterCaptureAudioRead, and add deprecation
- Add deprecation to BeforeCaptureAudioSent which previously had ChannelName added

[RN] Added 2 new delegates FOnVoiceChatAfterCaptureAudioReadDelegate2 and FOnVoiceChatBeforeCaptureAudioSentDelegate2, replacing FOnVoiceChatAfterCaptureAudioReadDelegate and FOnVoiceChatBeforeCaptureAudioSentDelegate respectively, adding a ChannelName param.

diff --git a/Engine/Plugins/Online/OnlineSubsystemEOS/Source/OnlineSubsystemEOS/Private/OnlineSubsystemEOS.cpp b/Engine/Plugins/Online/OnlineSubsystemEOS/Source/OnlineSubsystemEOS/Private/OnlineSubsystemEOS.cpp
--- a/Engine/Plugins/Online/OnlineSubsystemEOS/Source/OnlineSubsystemEOS/Private/OnlineSubsystemEOS.cpp
+++ b/Engine/Plugins/Online/OnlineSubsystemEOS/Source/OnlineSubsystemEOS/Private/OnlineSubsystemEOS.cpp
@@ -112,9 +112,9 @@ public:
 	virtual TSet<FString> GetTransmitChannels() const override { return VoiceChatUser.GetTransmitChannels(); }
 	virtual FDelegateHandle StartRecording(const FOnVoiceChatRecordSamplesAvailableDelegate::FDelegate& Delegate) override { return VoiceChatUser.StartRecording(Delegate); }
 	virtual void StopRecording(FDelegateHandle Handle) override { VoiceChatUser.StopRecording(Handle); }
-	virtual FDelegateHandle RegisterOnVoiceChatAfterCaptureAudioReadDelegate(const FOnVoiceChatAfterCaptureAudioReadDelegate::FDelegate& Delegate) override { return VoiceChatUser.RegisterOnVoiceChatAfterCaptureAudioReadDelegate(Delegate); }
+	virtual FDelegateHandle RegisterOnVoiceChatAfterCaptureAudioReadDelegate(const FOnVoiceChatAfterCaptureAudioReadDelegate2::FDelegate& Delegate) override { return VoiceChatUser.RegisterOnVoiceChatAfterCaptureAudioReadDelegate(Delegate); }
 	virtual void UnregisterOnVoiceChatAfterCaptureAudioReadDelegate(FDelegateHandle Handle) override { VoiceChatUser.UnregisterOnVoiceChatAfterCaptureAudioReadDelegate(Handle); }
-	virtual FDelegateHandle RegisterOnVoiceChatBeforeCaptureAudioSentDelegate(const FOnVoiceChatBeforeCaptureAudioSentDelegate::FDelegate& Delegate) override { return VoiceChatUser.RegisterOnVoiceChatBeforeCaptureAudioSentDelegate(Delegate); }
+	virtual FDelegateHandle RegisterOnVoiceChatBeforeCaptureAudioSentDelegate(const FOnVoiceChatBeforeCaptureAudioSentDelegate2::FDelegate& Delegate) override { return VoiceChatUser.RegisterOnVoiceChatBeforeCaptureAudioSentDelegate(Delegate); }
 	virtual void UnregisterOnVoiceChatBeforeCaptureAudioSentDelegate(FDelegateHandle Handle) override { VoiceChatUser.UnregisterOnVoiceChatBeforeCaptureAudioSentDelegate(Handle); }
 	virtual FDelegateHandle RegisterOnVoiceChatBeforeRecvAudioRenderedDelegate(const FOnVoiceChatBeforeRecvAudioRenderedDelegate::FDelegate& Delegate) override { return VoiceChatUser.RegisterOnVoiceChatBeforeRecvAudioRenderedDelegate(Delegate); }
 	virtual void UnregisterOnVoiceChatBeforeRecvAudioRenderedDelegate(FDelegateHandle Handle) override { VoiceChatUser.UnregisterOnVoiceChatBeforeRecvAudioRenderedDelegate(Handle); }
diff --git a/Engine/Plugins/Online/VoiceChat/EOSVoiceChat/Source/EOSVoiceChat/Private/EOSVoiceChat.cpp b/Engine/Plugins/Online/VoiceChat/EOSVoiceChat/Source/EOSVoiceChat/Private/EOSVoiceChat.cpp
--- a/Engine/Plugins/Online/VoiceChat/EOSVoiceChat/Source/EOSVoiceChat/Private/EOSVoiceChat.cpp
+++ b/Engine/Plugins/Online/VoiceChat/EOSVoiceChat/Source/EOSVoiceChat/Private/EOSVoiceChat.cpp
@@ -714,7 +714,7 @@ void FEOSVoiceChat::StopRecording(FDelegateHandle Handle)
 	GetVoiceChatUser().StopRecording(Handle);
 }
 
-FDelegateHandle FEOSVoiceChat::RegisterOnVoiceChatAfterCaptureAudioReadDelegate(const FOnVoiceChatAfterCaptureAudioReadDelegate::FDelegate& Delegate)
+FDelegateHandle FEOSVoiceChat::RegisterOnVoiceChatAfterCaptureAudioReadDelegate(const FOnVoiceChatAfterCaptureAudioReadDelegate2::FDelegate& Delegate)
 {
 	return GetVoiceChatUser().RegisterOnVoiceChatAfterCaptureAudioReadDelegate(Delegate);
 }
@@ -724,7 +724,7 @@ void FEOSVoiceChat::UnregisterOnVoiceChatAfterCaptureAudioReadDelegate(FDelegate
 	GetVoiceChatUser().UnregisterOnVoiceChatAfterCaptureAudioReadDelegate(Handle);
 }
 
-FDelegateHandle FEOSVoiceChat::RegisterOnVoiceChatBeforeCaptureAudioSentDelegate(const FOnVoiceChatBeforeCaptureAudioSentDelegate::FDelegate& Delegate)
+FDelegateHandle FEOSVoiceChat::RegisterOnVoiceChatBeforeCaptureAudioSentDelegate(const FOnVoiceChatBeforeCaptureAudioSentDelegate2::FDelegate& Delegate)
 {
 	return GetVoiceChatUser().RegisterOnVoiceChatBeforeCaptureAudioSentDelegate(Delegate);
 }
diff --git a/Engine/Plugins/Online/VoiceChat/EOSVoiceChat/Source/EOSVoiceChat/Private/EOSVoiceChatUser.cpp b/Engine/Plugins/Online/VoiceChat/EOSVoiceChat/Source/EOSVoiceChat/Private/EOSVoiceChatUser.cpp
--- a/Engine/Plugins/Online/VoiceChat/EOSVoiceChat/Source/EOSVoiceChat/Private/EOSVoiceChatUser.cpp
+++ b/Engine/Plugins/Online/VoiceChat/EOSVoiceChat/Source/EOSVoiceChat/Private/EOSVoiceChatUser.cpp
@@ -833,7 +833,7 @@ void FEOSVoiceChatUser::StopRecording(FDelegateHandle Handle)
 #endif
 }
 
-FDelegateHandle FEOSVoiceChatUser::RegisterOnVoiceChatAfterCaptureAudioReadDelegate(const FOnVoiceChatAfterCaptureAudioReadDelegate::FDelegate& Delegate)
+FDelegateHandle FEOSVoiceChatUser::RegisterOnVoiceChatAfterCaptureAudioReadDelegate(const FOnVoiceChatAfterCaptureAudioReadDelegate2::FDelegate& Delegate)
 {
 	FScopeLock Lock(&BeforeCaptureAudioSentLock);
 
@@ -847,7 +847,7 @@ void FEOSVoiceChatUser::UnregisterOnVoiceChatAfterCaptureAudioReadDelegate(FDele
 	OnVoiceChatAfterCaptureAudioReadDelegate.Remove(Handle);
 }
 
-FDelegateHandle FEOSVoiceChatUser::RegisterOnVoiceChatBeforeCaptureAudioSentDelegate(const FOnVoiceChatBeforeCaptureAudioSentDelegate::FDelegate& Delegate)
+FDelegateHandle FEOSVoiceChatUser::RegisterOnVoiceChatBeforeCaptureAudioSentDelegate(const FOnVoiceChatBeforeCaptureAudioSentDelegate2::FDelegate& Delegate)
 {
 	FScopeLock Lock(&BeforeCaptureAudioSentLock);
 
@@ -2234,12 +2234,12 @@ void FEOSVoiceChatUser::OnChannelAudioBeforeSend(const EOS_RTCAudio_AudioBeforeS
 				}
 			}
 
+			const FString ChannelName = UTF8_TO_TCHAR(CallbackInfo->RoomName);
+
 			FScopeLock Lock(&BeforeCaptureAudioSentLock);
 
 			// Allow any processes to modify audio through DSP effects processing
-			OnVoiceChatAfterCaptureAudioReadDelegate.Broadcast(WriteableSamples, Buffer->SampleRate, Buffer->Channels);
-
-			const FString ChannelName = UTF8_TO_TCHAR(CallbackInfo->RoomName);
+			OnVoiceChatAfterCaptureAudioReadDelegate.Broadcast(ChannelName, WriteableSamples, Buffer->SampleRate, Buffer->Channels);
 
 			OnVoiceChatBeforeCaptureAudioSentDelegate.Broadcast(ChannelName, Samples, Buffer->SampleRate, Buffer->Channels, bSpeaking);
 
diff --git a/Engine/Plugins/Online/VoiceChat/EOSVoiceChat/Source/EOSVoiceChat/Public/EOSVoiceChat.h b/Engine/Plugins/Online/VoiceChat/EOSVoiceChat/Source/EOSVoiceChat/Public/EOSVoiceChat.h
--- a/Engine/Plugins/Online/VoiceChat/EOSVoiceChat/Source/EOSVoiceChat/Public/EOSVoiceChat.h
+++ b/Engine/Plugins/Online/VoiceChat/EOSVoiceChat/Source/EOSVoiceChat/Public/EOSVoiceChat.h
@@ -128,9 +128,9 @@ public:
 	virtual TSet<FString> GetTransmitChannels() const override;
 	virtual FDelegateHandle StartRecording(const FOnVoiceChatRecordSamplesAvailableDelegate::FDelegate& Delegate) override;
 	virtual void StopRecording(FDelegateHandle Handle) override;
-	virtual FDelegateHandle RegisterOnVoiceChatAfterCaptureAudioReadDelegate(const FOnVoiceChatAfterCaptureAudioReadDelegate::FDelegate& Delegate) override;
+	virtual FDelegateHandle RegisterOnVoiceChatAfterCaptureAudioReadDelegate(const FOnVoiceChatAfterCaptureAudioReadDelegate2::FDelegate& Delegate) override;
 	virtual void UnregisterOnVoiceChatAfterCaptureAudioReadDelegate(FDelegateHandle Handle) override;
-	virtual FDelegateHandle RegisterOnVoiceChatBeforeCaptureAudioSentDelegate(const FOnVoiceChatBeforeCaptureAudioSentDelegate::FDelegate& Delegate) override;
+	virtual FDelegateHandle RegisterOnVoiceChatBeforeCaptureAudioSentDelegate(const FOnVoiceChatBeforeCaptureAudioSentDelegate2::FDelegate& Delegate) override;
 	virtual void UnregisterOnVoiceChatBeforeCaptureAudioSentDelegate(FDelegateHandle Handle) override;
 	virtual FDelegateHandle RegisterOnVoiceChatBeforeRecvAudioRenderedDelegate(const FOnVoiceChatBeforeRecvAudioRenderedDelegate::FDelegate& Delegate) override;
 	virtual void UnregisterOnVoiceChatBeforeRecvAudioRenderedDelegate(FDelegateHandle Handle) override;
diff --git a/Engine/Plugins/Online/VoiceChat/EOSVoiceChat/Source/EOSVoiceChat/Public/EOSVoiceChatUser.h b/Engine/Plugins/Online/VoiceChat/EOSVoiceChat/Source/EOSVoiceChat/Public/EOSVoiceChatUser.h
--- a/Engine/Plugins/Online/VoiceChat/EOSVoiceChat/Source/EOSVoiceChat/Public/EOSVoiceChatUser.h
+++ b/Engine/Plugins/Online/VoiceChat/EOSVoiceChat/Source/EOSVoiceChat/Public/EOSVoiceChatUser.h
@@ -76,9 +76,9 @@ public:
 	EOSVOICECHAT_API virtual TSet<FString> GetTransmitChannels() const override;
 	EOSVOICECHAT_API virtual FDelegateHandle StartRecording(const FOnVoiceChatRecordSamplesAvailableDelegate::FDelegate& Delegate) override;
 	EOSVOICECHAT_API virtual void StopRecording(FDelegateHandle Handle) override;
-	EOSVOICECHAT_API virtual FDelegateHandle RegisterOnVoiceChatAfterCaptureAudioReadDelegate(const FOnVoiceChatAfterCaptureAudioReadDelegate::FDelegate& Delegate) override;
+	EOSVOICECHAT_API virtual FDelegateHandle RegisterOnVoiceChatAfterCaptureAudioReadDelegate(const FOnVoiceChatAfterCaptureAudioReadDelegate2::FDelegate& Delegate) override;
 	EOSVOICECHAT_API virtual void UnregisterOnVoiceChatAfterCaptureAudioReadDelegate(FDelegateHandle Handle) override;
-	EOSVOICECHAT_API virtual FDelegateHandle RegisterOnVoiceChatBeforeCaptureAudioSentDelegate(const FOnVoiceChatBeforeCaptureAudioSentDelegate::FDelegate& Delegate) override;
+	EOSVOICECHAT_API virtual FDelegateHandle RegisterOnVoiceChatBeforeCaptureAudioSentDelegate(const FOnVoiceChatBeforeCaptureAudioSentDelegate2::FDelegate& Delegate) override;
 	EOSVOICECHAT_API virtual void UnregisterOnVoiceChatBeforeCaptureAudioSentDelegate(FDelegateHandle Handle) override;
 	EOSVOICECHAT_API virtual FDelegateHandle RegisterOnVoiceChatBeforeRecvAudioRenderedDelegate(const FOnVoiceChatBeforeRecvAudioRenderedDelegate::FDelegate& Delegate) override;
 	EOSVOICECHAT_API virtual void UnregisterOnVoiceChatBeforeRecvAudioRenderedDelegate(FDelegateHandle Handle) override;
@@ -280,14 +280,18 @@ protected:
 	FOnVoiceChatPlayerRemovedDelegate OnVoiceChatPlayerRemovedDelegate;
 	FOnVoiceChatCallStatsUpdatedDelegate OnVoiceChatCallStatsUpdatedDelegate;
 
-	// IVoiceChatUser Recording Delegates
+	// Begin IVoiceChatUser Recording Delegates
 	FCriticalSection AudioRecordLock;
 	FOnVoiceChatRecordSamplesAvailableDelegate OnVoiceChatRecordSamplesAvailableDelegate;
-	FOnVoiceChatAfterCaptureAudioReadDelegate OnVoiceChatAfterCaptureAudioReadDelegate;
+
+	// Note BeforeCaptureAudioSentLock is used for CaptureAudioRead too as they fire together
 	FCriticalSection BeforeCaptureAudioSentLock;
-	FOnVoiceChatBeforeCaptureAudioSentDelegate OnVoiceChatBeforeCaptureAudioSentDelegate;
+	FOnVoiceChatAfterCaptureAudioReadDelegate2 OnVoiceChatAfterCaptureAudioReadDelegate;
+	FOnVoiceChatBeforeCaptureAudioSentDelegate2 OnVoiceChatBeforeCaptureAudioSentDelegate;
+
 	FCriticalSection BeforeRecvAudioRenderedLock;
 	FOnVoiceChatBeforeRecvAudioRenderedDelegate OnVoiceChatBeforeRecvAudioRenderedDelegate;
+	// End IVoiceChatUser Recording Delegates
 
 	FEOSVoiceChat& EOSVoiceChat;
 
diff --git a/Engine/Plugins/Online/VoiceChat/VoiceChat/Source/Public/VoiceChat.h b/Engine/Plugins/Online/VoiceChat/VoiceChat/Source/Public/VoiceChat.h
--- a/Engine/Plugins/Online/VoiceChat/VoiceChat/Source/Public/VoiceChat.h
+++ b/Engine/Plugins/Online/VoiceChat/VoiceChat/Source/Public/VoiceChat.h
@@ -111,10 +111,13 @@ DECLARE_MULTICAST_DELEGATE_ThreeParams(FOnVoiceChatPlayerMuteUpdatedDelegate, co
 DECLARE_MULTICAST_DELEGATE_ThreeParams(FOnVoiceChatPlayerVolumeUpdatedDelegate, const FString& /* ChannelName */, const FString& /* PlayerName */, float /* Volume */);
 
 DECLARE_MULTICAST_DELEGATE_ThreeParams(FOnVoiceChatRecordSamplesAvailableDelegate, TArrayView<const int16> /* PcmSamples */, int /* SampleRate */, int /* Channels */);
-DECLARE_MULTICAST_DELEGATE_ThreeParams(FOnVoiceChatAfterCaptureAudioReadDelegate, TArrayView<int16> /* PcmSamples */, int /* SampleRate */, int /* Channels */);
-DECLARE_MULTICAST_DELEGATE_FiveParams(FOnVoiceChatBeforeCaptureAudioSentDelegate, const FString& /* ChannelName */, TArrayView<const int16> /* PcmSamples */, int /* SampleRate */, int /* Channels */, bool /* bIsSpeaking */);
+DECLARE_MULTICAST_DELEGATE_FourParams(FOnVoiceChatAfterCaptureAudioReadDelegate2, const FString& /* ChannelName */, TArrayView<int16> /* PcmSamples */, int /* SampleRate */, int /* Channels */);
+DECLARE_MULTICAST_DELEGATE_FiveParams(FOnVoiceChatBeforeCaptureAudioSentDelegate2, const FString& /* ChannelName */, TArrayView<const int16> /* PcmSamples */, int /* SampleRate */, int /* Channels */, bool /* bIsSpeaking */);
 DECLARE_MULTICAST_DELEGATE_SixParams(FOnVoiceChatBeforeRecvAudioRenderedDelegate, TArrayView<int16> /* PcmSamples */, int /* SampleRate */, int /* Channels */, bool /* bIsSilence */, const FString& /*ChannelName*/, const FString& /* PlayerName */);
 
+DECLARE_MULTICAST_DELEGATE_ThreeParams(FOnVoiceChatAfterCaptureAudioReadDelegate, TArrayView<int16> /* PcmSamples */, int /* SampleRate */, int /* Channels */);
+DECLARE_MULTICAST_DELEGATE_FourParams(FOnVoiceChatBeforeCaptureAudioSentDelegate, TArrayView<const int16> /* PcmSamples */, int /* SampleRate */, int /* Channels */, bool /* bIsSpeaking */);
+
 class IVoiceChatUser
 {
 public:
@@ -624,15 +627,33 @@ public:
 	/**
 	 * Delegates called from audio thread.
 	 */
-	virtual FDelegateHandle RegisterOnVoiceChatAfterCaptureAudioReadDelegate(const FOnVoiceChatAfterCaptureAudioReadDelegate::FDelegate& Delegate) = 0;
+	virtual FDelegateHandle RegisterOnVoiceChatAfterCaptureAudioReadDelegate(const FOnVoiceChatAfterCaptureAudioReadDelegate2::FDelegate& Delegate) = 0;
 	virtual void UnregisterOnVoiceChatAfterCaptureAudioReadDelegate(FDelegateHandle Handle) = 0;
 
-	virtual FDelegateHandle RegisterOnVoiceChatBeforeCaptureAudioSentDelegate(const FOnVoiceChatBeforeCaptureAudioSentDelegate::FDelegate& Delegate) = 0;
+	virtual FDelegateHandle RegisterOnVoiceChatBeforeCaptureAudioSentDelegate(const FOnVoiceChatBeforeCaptureAudioSentDelegate2::FDelegate& Delegate) = 0;
 	virtual void UnregisterOnVoiceChatBeforeCaptureAudioSentDelegate(FDelegateHandle Handle) = 0;
 
 	virtual FDelegateHandle RegisterOnVoiceChatBeforeRecvAudioRenderedDelegate(const FOnVoiceChatBeforeRecvAudioRenderedDelegate::FDelegate& Delegate) = 0;
 	virtual void UnregisterOnVoiceChatBeforeRecvAudioRenderedDelegate(FDelegateHandle Handle) = 0;
 
+	UE_DEPRECATED(5.6, "Please pass a FOnVoiceChatBeforeCaptureAudioSentDelegate2, which adds the ChannelName param")
+	virtual FDelegateHandle RegisterOnVoiceChatAfterCaptureAudioReadDelegate(const FOnVoiceChatAfterCaptureAudioReadDelegate::FDelegate& Delegate) final
+	{
+		return RegisterOnVoiceChatAfterCaptureAudioReadDelegate(FOnVoiceChatAfterCaptureAudioReadDelegate2::FDelegate::CreateLambda([Delegate](const FString& /* ChannelName */, TArrayView<int16> PcmSamples, int SampleRate, int Channels)
+			{
+				Delegate.ExecuteIfBound(PcmSamples, SampleRate, Channels);
+			}));
+	}
+
+	UE_DEPRECATED(5.6, "Please pass a FOnVoiceChatBeforeCaptureAudioSentDelegate2, which adds the ChannelName param")
+		virtual FDelegateHandle RegisterOnVoiceChatBeforeCaptureAudioSentDelegate(const FOnVoiceChatBeforeCaptureAudioSentDelegate::FDelegate& Delegate) final
+	{
+		return RegisterOnVoiceChatBeforeCaptureAudioSentDelegate(FOnVoiceChatBeforeCaptureAudioSentDelegate2::FDelegate::CreateLambda([Delegate](const FString& /* ChannelName */, TArrayView<const int16> PcmSamples, int SampleRate, int Channels, bool bIsSpeaking)
+			{
+				Delegate.ExecuteIfBound(PcmSamples, SampleRate, Channels, bIsSpeaking);
+			}));
+	}
+
 	/**
 	 * Generate Login/Join tokens. Should be used for development purposes only
 	 */
