- Add LexToString for some VoiceChat enums
- Convert VoiceChat to a proper module to house the LexToString methods

#changelist validated

diff --git a/Engine/Plugins/Online/OnlineSubsystemEOS/OnlineSubsystemEOS.uplugin b/Engine/Plugins/Online/OnlineSubsystemEOS/OnlineSubsystemEOS.uplugin
--- a/Engine/Plugins/Online/OnlineSubsystemEOS/OnlineSubsystemEOS.uplugin
+++ b/Engine/Plugins/Online/OnlineSubsystemEOS/OnlineSubsystemEOS.uplugin
@@ -49,6 +49,10 @@
 		{
 			"Name": "EOSOverlayInputProvider",
 			"Enabled": true
+		},
+		{
+			"Name": "VoiceChat",
+			"Enabled": true
 		}
 	]
 }
diff --git a/Engine/Plugins/Online/OnlineSubsystemEOS/Source/OnlineSubsystemEOS/OnlineSubsystemEOS.Build.cs b/Engine/Plugins/Online/OnlineSubsystemEOS/Source/OnlineSubsystemEOS/OnlineSubsystemEOS.Build.cs
--- a/Engine/Plugins/Online/OnlineSubsystemEOS/Source/OnlineSubsystemEOS/OnlineSubsystemEOS.Build.cs
+++ b/Engine/Plugins/Online/OnlineSubsystemEOS/Source/OnlineSubsystemEOS/OnlineSubsystemEOS.Build.cs
@@ -16,13 +16,6 @@ public class OnlineSubsystemEOS : ModuleRules
 			}
 		);
 
-		PrivateIncludePathModuleNames.AddRange(
-			new string[]
-			{
-				"VoiceChat",
-			}
-		);
-
 		PrivateDependencyModuleNames.AddRange(
 			new string[] {
 				"Core",
@@ -31,12 +24,13 @@ public class OnlineSubsystemEOS : ModuleRules
 				"Engine",
 				"EOSVoiceChat",
 				"Json",
+				"NetCore",
 				"OnlineBase",
 				"OnlineSubsystem",
 				"OnlineSubsystemUtils",
 				"Sockets",
 				"SocketSubsystemEOS",
-				"NetCore"
+				"VoiceChat"
 			}
 		);
 
diff --git a/Engine/Plugins/Online/VoiceChat/EOSVoiceChat/EOSVoiceChat.uplugin b/Engine/Plugins/Online/VoiceChat/EOSVoiceChat/EOSVoiceChat.uplugin
--- a/Engine/Plugins/Online/VoiceChat/EOSVoiceChat/EOSVoiceChat.uplugin
+++ b/Engine/Plugins/Online/VoiceChat/EOSVoiceChat/EOSVoiceChat.uplugin
@@ -25,6 +25,10 @@
 		{
 			"Name": "EOSShared",
 			"Enabled": true
+		},
+		{
+			"Name": "VoiceChat",
+			"Enabled": true
 		}
 	]
 }
diff --git a/Engine/Plugins/Online/VoiceChat/EOSVoiceChat/Source/EOSVoiceChat/EOSVoiceChat.Build.cs b/Engine/Plugins/Online/VoiceChat/EOSVoiceChat/Source/EOSVoiceChat/EOSVoiceChat.Build.cs
--- a/Engine/Plugins/Online/VoiceChat/EOSVoiceChat/Source/EOSVoiceChat/EOSVoiceChat.Build.cs
+++ b/Engine/Plugins/Online/VoiceChat/EOSVoiceChat/Source/EOSVoiceChat/EOSVoiceChat.Build.cs
@@ -46,14 +46,8 @@ namespace UnrealBuildTool.Rules
 			PublicDependencyModuleNames.AddRange(
 				new string[]
 				{
-					"Json"
-				}
-			);
-
-			PublicIncludePathModuleNames.AddRange(
-				new string[]
-				{
-					"VoiceChat",
+					"Json",
+					"VoiceChat"
 				}
 			);
 
diff --git a/Engine/Plugins/Online/VoiceChat/VoiceChat/Source/Private/VoiceChat.cpp b/Engine/Plugins/Online/VoiceChat/VoiceChat/Source/Private/VoiceChat.cpp
new file mode 100644
--- /dev/null
+++ b/Engine/Plugins/Online/VoiceChat/VoiceChat/Source/Private/VoiceChat.cpp
@@ -0,0 +1,37 @@
+// Copyright Epic Games, Inc. All Rights Reserved.
+
+#include "VoiceChat.h"
+
+const TCHAR* LexToString(const EVoiceChatChannelType Value)
+{
+	switch (Value)
+	{
+		case EVoiceChatChannelType::NonPositional: return TEXT("NonPositional");
+		case EVoiceChatChannelType::Positional: return TEXT("Positional");
+		case EVoiceChatChannelType::Echo: return TEXT("Echo");
+		default: checkNoEntry(); return TEXT("Invalid");
+	}
+}
+
+const TCHAR* LexToString(const EVoiceChatAttenuationModel Value)
+{
+	switch (Value)
+	{
+		case EVoiceChatAttenuationModel::None: return TEXT("None");
+		case EVoiceChatAttenuationModel::InverseByDistance: return TEXT("InverseByDistance");
+		case EVoiceChatAttenuationModel::LinearByDistance: return TEXT("LinearByDistance");
+		case EVoiceChatAttenuationModel::ExponentialByDistance: return TEXT("ExponentialByDistance");
+		default: checkNoEntry(); return TEXT("Invalid");
+	}
+}
+
+const TCHAR* LexToString(const EVoiceChatTransmitMode Value)
+{
+	switch (Value)
+	{
+		case EVoiceChatTransmitMode::None: return TEXT("None");
+		case EVoiceChatTransmitMode::All: return TEXT("All");
+		case EVoiceChatTransmitMode::SpecificChannels: return TEXT("SpecificChannels");
+		default: checkNoEntry(); return TEXT("Invalid");
+	}
+}
\ No newline at end of file
diff --git a/Engine/Plugins/Online/VoiceChat/VoiceChat/Source/Private/VoiceChatModule.cpp b/Engine/Plugins/Online/VoiceChat/VoiceChat/Source/Private/VoiceChatModule.cpp
new file mode 100644
--- /dev/null
+++ b/Engine/Plugins/Online/VoiceChat/VoiceChat/Source/Private/VoiceChatModule.cpp
@@ -0,0 +1,10 @@
+// Copyright Epic Games, Inc. All Rights Reserved.
+
+#include "Modules/ModuleInterface.h"
+#include "Modules/ModuleManager.h"
+
+class FVoiceChatModule : public IModuleInterface
+{
+};
+
+IMPLEMENT_MODULE(FVoiceChatModule, VoiceChat)
\ No newline at end of file
diff --git a/Engine/Plugins/Online/VoiceChat/VoiceChat/Source/Public/VoiceChat.h b/Engine/Plugins/Online/VoiceChat/VoiceChat/Source/Public/VoiceChat.h
--- a/Engine/Plugins/Online/VoiceChat/VoiceChat/Source/Public/VoiceChat.h
+++ b/Engine/Plugins/Online/VoiceChat/VoiceChat/Source/Public/VoiceChat.h
@@ -18,6 +18,8 @@ enum class EVoiceChatChannelType
 	Echo
 };
 
+VOICECHAT_API const TCHAR* LexToString(const EVoiceChatChannelType Value);
+
 enum class EVoiceChatAttenuationModel
 {
 	/** No attenuation is applied. The audio will drop to 0 at MaxDistance */
@@ -30,6 +32,8 @@ enum class EVoiceChatAttenuationModel
 	ExponentialByDistance
 };
 
+VOICECHAT_API const TCHAR* LexToString(const EVoiceChatAttenuationModel Value);
+
 struct FVoiceChatChannel3dProperties
 {
 	/** The model used to determine how loud audio is at different distances */
@@ -52,6 +56,8 @@ enum class EVoiceChatTransmitMode
 	SpecificChannels,
 };
 
+VOICECHAT_API const TCHAR* LexToString(const EVoiceChatTransmitMode Value);
+
 struct FVoiceChatCallStats
 {
 	/** Call length in seconds */
diff --git a/Engine/Plugins/Online/VoiceChat/VoiceChat/Source/VoiceChat.Build.cs b/Engine/Plugins/Online/VoiceChat/VoiceChat/Source/VoiceChat.Build.cs
--- a/Engine/Plugins/Online/VoiceChat/VoiceChat/Source/VoiceChat.Build.cs
+++ b/Engine/Plugins/Online/VoiceChat/VoiceChat/Source/VoiceChat.Build.cs
@@ -8,8 +8,6 @@ namespace UnrealBuildTool.Rules
 	{
 		public VoiceChat(ReadOnlyTargetRules Target) : base(Target)
 		{
-			// External so we don't build a DLL for this, it is header only.
-			Type = ModuleType.External;
 			CppCompileWarningSettings.ModuleIncludePathWarningLevel = WarningLevel.Off;
 
 			PublicIncludePaths.Add(Path.Combine(ModuleDirectory, "Public"));
